Allow the game input windows/game input base modules to exist on server builds. This fixes an issue where the module would be requested but not exist on the server.

#preflgiht 65cba36ae15916236d979412
#rnx

diff --git a/Engine/Plugins/Experimental/GameInputWindows/GameInputWindows.uplugin b/Engine/Plugins/Experimental/GameInputWindows/GameInputWindows.uplugin
--- a/Engine/Plugins/Experimental/GameInputWindows/GameInputWindows.uplugin
+++ b/Engine/Plugins/Experimental/GameInputWindows/GameInputWindows.uplugin
@@ -18,9 +18,8 @@
 	"Modules": [
 		{
 			"Name": "GameInputWindows",
-			"Type": "Runtime",
+			"Type": "RuntimeNoCommandlet",
 			"HasExplicitPlatforms": true,
-			"TargetDenyList": [ "Server" ],
 			"PlatformAllowList": [
 				"Win64"
 			]
@@ -29,8 +28,7 @@
 	"Plugins": [
 		{
 			"Name": "GameInput",
-			"Enabled": true,
-			"TargetDenyList": [ "Server" ]
+			"Enabled": true
 		}
 	]
 }
\ No newline at end of file
