Make FGameInputWindowsInputDevice not final so that people can extend and build upon it for any custom input device processing if they would like to. This would be a rather advanced input setup, but it should be possible :)

#rnx

diff --git a/Engine/Plugins/Experimental/GameInputWindows/Source/GameInputWindows/Private/GameInputWindowsDevice.h b/Engine/Plugins/Experimental/GameInputWindows/Source/GameInputWindows/Private/GameInputWindowsDevice.h
--- a/Engine/Plugins/Experimental/GameInputWindows/Source/GameInputWindows/Private/GameInputWindowsDevice.h
+++ b/Engine/Plugins/Experimental/GameInputWindows/Source/GameInputWindows/Private/GameInputWindowsDevice.h
@@ -6,7 +6,7 @@
 
 #include "IGameInputDeviceInterface.h"
 
-class FGameInputWindowsInputDevice final : public IGameInputDeviceInterface
+class FGameInputWindowsInputDevice : public IGameInputDeviceInterface
 {
 public:
 	explicit FGameInputWindowsInputDevice(const TSharedRef<FGenericApplicationMessageHandler>& InMessageHandler, struct IGameInput* InGameInput);
