Add some more info about where to get info on why GameInput may have failed to build.

#rnx

diff --git a/Engine/Plugins/Experimental/GameInputWindows/Source/GameInputWindows/Private/GameInputWindowsModule.cpp b/Engine/Plugins/Experimental/GameInputWindows/Source/GameInputWindows/Private/GameInputWindowsModule.cpp
--- a/Engine/Plugins/Experimental/GameInputWindows/Source/GameInputWindows/Private/GameInputWindowsModule.cpp
+++ b/Engine/Plugins/Experimental/GameInputWindows/Source/GameInputWindows/Private/GameInputWindowsModule.cpp
@@ -69,7 +69,7 @@ TSharedPtr<IInputDevice> FGameInputWindowsModule::CreateInputDevice(const TShare
 		}
 	}
 #else
-	UE_LOG(LogGameInput, Error, TEXT("[FGameInputWindowsModule] Failed to create a GameInput device! GAME_INPUT_SUPPORT is false!"));
+	UE_LOG(LogGameInput, Error, TEXT("[FGameInputWindowsModule] Failed to create a GameInput device! GAME_INPUT_SUPPORT is false! See GameInputBase.build.cs and the build log for more info."));
 #endif	// GAME_INPUT_SUPPORT
 
 	return nullptr;
