Grab the FGameInputBaseModule module earlier in FGameInputWindowsModule to ensure that it Initializes before we call FGameInputBaseModule::GetGameInput. This ensures that FGameInputBaseModule is loaded so that we can properly initialize.

This fixes guitar peripherals not working on Windows in Packages.

#rnx

diff --git a/Engine/Plugins/Experimental/GameInputWindows/Source/GameInputWindows/Private/GameInputWindowsModule.cpp b/Engine/Plugins/Experimental/GameInputWindows/Source/GameInputWindows/Private/GameInputWindowsModule.cpp
--- a/Engine/Plugins/Experimental/GameInputWindows/Source/GameInputWindows/Private/GameInputWindowsModule.cpp
+++ b/Engine/Plugins/Experimental/GameInputWindows/Source/GameInputWindows/Private/GameInputWindowsModule.cpp
@@ -29,6 +29,8 @@ bool FGameInputWindowsModule::IsAvailable()
 TSharedPtr<IInputDevice> FGameInputWindowsModule::CreateInputDevice(const TSharedRef<FGenericApplicationMessageHandler>& InMessageHandler)
 {
 #if GAME_INPUT_SUPPORT
+	//Get FGameInputBaseModule before anything else to make sure that it is loaded
+	FGameInputBaseModule& BaseModule = FGameInputBaseModule::Get();
 	IGameInput* BaseGameInputPtr = FGameInputBaseModule::GetGameInput();
 	TSharedPtr<FGameInputWindowsInputDevice> DevicePtr = MakeShared<FGameInputWindowsInputDevice>(InMessageHandler, BaseGameInputPtr);
 	
@@ -37,8 +39,7 @@ TSharedPtr<IInputDevice> FGameInputWindowsModule::CreateInputDevice(const TShare
 		DevicePtr->Initialize();
 	}
 	else
-	{
-		FGameInputBaseModule& BaseModule = FGameInputBaseModule::Get();
+	{	
 		BaseModule.OnGameInputCreation.AddLambda([&DevicePtr](IGameInput* GameInputPtr)
 		{
 			if (ensure(DevicePtr.IsValid()))
