[Game Input Windows] Move the Game Input for Windows plugin out of the experimental folder. This marks its release to "Beta".

#rnx

diff --git a/Engine/Plugins/Experimental/GameInputWindows/GameInputWindows.uplugin b/Engine/Plugins/Experimental/GameInputWindows/GameInputWindows.uplugin
deleted file mode 100644
diff --git a/Engine/Plugins/Experimental/GameInputWindows/Source/GameInputWindows/GameInputWindows.build.cs b/Engine/Plugins/Experimental/GameInputWindows/Source/GameInputWindows/GameInputWindows.build.cs
deleted file mode 100644
diff --git a/Engine/Plugins/Experimental/GameInputWindows/Source/GameInputWindows/Private/GameInputWindowsDevice.cpp b/Engine/Plugins/Experimental/GameInputWindows/Source/GameInputWindows/Private/GameInputWindowsDevice.cpp
deleted file mode 100644
diff --git a/Engine/Plugins/Experimental/GameInputWindows/Source/GameInputWindows/Private/GameInputWindowsDevice.h b/Engine/Plugins/Experimental/GameInputWindows/Source/GameInputWindows/Private/GameInputWindowsDevice.h
deleted file mode 100644
diff --git a/Engine/Plugins/Experimental/GameInputWindows/Source/GameInputWindows/Private/GameInputWindowsModule.cpp b/Engine/Plugins/Experimental/GameInputWindows/Source/GameInputWindows/Private/GameInputWindowsModule.cpp
deleted file mode 100644
diff --git a/Engine/Plugins/Experimental/GameInputWindows/Source/GameInputWindows/Private/GameInputWindowsModule.h b/Engine/Plugins/Experimental/GameInputWindows/Source/GameInputWindows/Private/GameInputWindowsModule.h
deleted file mode 100644
diff --git a/Engine/Plugins/GameInputWindows/GameInputWindows.uplugin b/Engine/Plugins/GameInputWindows/GameInputWindows.uplugin
new file mode 100644
--- /dev/null
+++ b/Engine/Plugins/GameInputWindows/GameInputWindows.uplugin
@@ -0,0 +1,34 @@
+{
+	"FileVersion": 3,
+	"FriendlyName": "Game Input (Windows)",
+	"Version": 1,
+	"VersionName": "1.0",
+	"CreatedBy": "Epic Games, Inc.",
+	"CreatedByURL": "https://epicgames.com",
+	"Description": "GameInput is a next-generation input API that exposes input devices of all kinds through a single consistent interface.",
+	"Category": "Input Devices",
+	"EnabledByDefault": false,
+	"IsBetaVersion": true,
+	"CanContainContent": false,
+	"HasExplicitPlatforms": true,
+	"bRequiresBuildPlatform": false,
+	"SupportedTargetPlatforms": [
+		"Win64"
+	],
+	"Modules": [
+		{
+			"Name": "GameInputWindows",
+			"Type": "RuntimeNoCommandlet",
+			"HasExplicitPlatforms": true,
+			"PlatformAllowList": [
+				"Win64"
+			]
+		}
+	],
+	"Plugins": [
+		{
+			"Name": "GameInput",
+			"Enabled": true
+		}
+	]
+}
\ No newline at end of file
diff --git a/Engine/Plugins/GameInputWindows/Source/GameInputWindows/GameInputWindows.build.cs b/Engine/Plugins/GameInputWindows/Source/GameInputWindows/GameInputWindows.build.cs
new file mode 100644
--- /dev/null
+++ b/Engine/Plugins/GameInputWindows/Source/GameInputWindows/GameInputWindows.build.cs
@@ -0,0 +1,26 @@
+// Copyright Epic Games, Inc. All Rights Reserved.
+
+using System;
+
+namespace UnrealBuildTool.Rules
+{
+	public class GameInputWindows : ModuleRules
+	{
+		public GameInputWindows(ReadOnlyTargetRules Target) : base(Target)
+		{
+			PrivateDependencyModuleNames.AddRange(
+				new string[]
+				{
+					"Core",
+                    "ApplicationCore",
+                    "Engine",
+                    "InputCore",
+                    "InputDevice",				
+                    "CoreUObject",
+                    "DeveloperSettings",
+					"GameInputBase",
+				}
+			);
+		}
+	}
+}
\ No newline at end of file
diff --git a/Engine/Plugins/GameInputWindows/Source/GameInputWindows/Private/GameInputWindowsDevice.cpp b/Engine/Plugins/GameInputWindows/Source/GameInputWindows/Private/GameInputWindowsDevice.cpp
new file mode 100644
--- /dev/null
+++ b/Engine/Plugins/GameInputWindows/Source/GameInputWindows/Private/GameInputWindowsDevice.cpp
@@ -0,0 +1,216 @@
+// Copyright Epic Games, Inc. All Rights Reserved.
+
+#include "GameInputWindowsDevice.h"
+
+#if GAME_INPUT_SUPPORT
+
+#include "GameInputUtils.h"
+#include "GameInputLogging.h"
+#include "HAL/IConsoleManager.h" // FAutoConsoleVariableRef
+#include "GenericPlatform/IInputInterface.h"
+#include "GameInputDeveloperSettings.h"
+
+FGameInputWindowsInputDevice::FGameInputWindowsInputDevice(const TSharedRef<FGenericApplicationMessageHandler>& InMessageHandler, IGameInput* InGameInput)
+	: IGameInputDeviceInterface(InMessageHandler, InGameInput)
+{
+#if WITH_EDITOR
+	SetupEditorSettingListener();
+#endif
+}
+
+FGameInputWindowsInputDevice::~FGameInputWindowsInputDevice()
+{
+#if WITH_EDITOR
+	CleanupEditorSettingListener();
+#endif
+}
+
+void FGameInputWindowsInputDevice::SetGameInputAndReinitialize(IGameInput* InGameInput)
+{
+	GameInput = InGameInput;
+	Initialize();
+}
+
+GameInputKind FGameInputWindowsInputDevice::GetCurrentGameInputKindSupport() const
+{
+	GameInputKind RegisterInputKindMask = IGameInputDeviceInterface::GetCurrentGameInputKindSupport();
+
+	// For now, we want to explicitly make sure that Keyboard and Mouse are NOT being processed by GameInput on windows.
+	// This is because the WindowsApplication is doing to do all the processing of these types for us already, and 
+	// having the Game Input device plugin process them too would result in "double" events and mouse accumulation.
+
+#if !UE_BUILD_SHIPPING
+	static bool bLogOnce = false;
+	if (!bLogOnce)
+	{
+		if ((RegisterInputKindMask & GameInputKindKeyboard) != 0)
+		{
+			UE_LOG(LogGameInput, Log, TEXT("[FGameInputWindowsInputDevice::GetCurrentGameInputKindSupport] Keyboard support was requested, but is not currently supported via the GameInput plugin on Windows."));
+		}
+
+		if ((RegisterInputKindMask & GameInputKindMouse) != 0)
+		{
+			UE_LOG(LogGameInput, Log, TEXT("[FGameInputWindowsInputDevice::GetCurrentGameInputKindSupport] Mouse support was requested, but is not currently supported via the GameInput plugin on Windows."));
+		}
+		bLogOnce = true;
+	}
+#endif	// !UE_BUILD_SHIPPING
+
+	// Clear the bits on these types to make sure we don't process any of them
+	RegisterInputKindMask &= ~GameInputKindKeyboard;
+	RegisterInputKindMask &= ~GameInputKindMouse;
+
+	return RegisterInputKindMask;
+}
+
+void FGameInputWindowsInputDevice::HandleDeviceDisconnected(IGameInputDevice* Device, uint64 Timestamp)
+{	
+	if (Device)
+	{
+		if (FGameInputDeviceContainer* Data = GetDeviceData(Device))
+		{
+			// Clear any input state that might be related to this device
+			Data->ClearInputState(GameInput);
+
+			// Set it's device to nullptr because it is now disconnected
+			Data->SetGameInputDevice(nullptr);
+			UE_LOG(LogGameInput, Log, TEXT("Game Input Device '%s' Disconnected Successfully at Input Device ID '%d'"), *UE::GameInput::LexToString(Device), Data->GetDeviceId().GetId());
+
+			// Remap this device to the "unpaired" user because it has been disconnected
+			const FPlatformUserId NewUserToAssign = IPlatformInputDeviceMapper::Get().GetUserForUnpairedInputDevices();
+			const FInputDeviceId DeviceId = Data->GetDeviceId();
+
+			if (DeviceId.IsValid())
+			{
+				const bool bSuccess = IPlatformInputDeviceMapper::Get().Internal_MapInputDeviceToUser(DeviceId, NewUserToAssign, EInputDeviceConnectionState::Disconnected);
+				if (bSuccess)
+				{
+					Data->SetPlatformUserId(NewUserToAssign);
+				}
+			}			
+		}
+		else
+		{
+			UE_LOG(LogGameInput, Error, TEXT("Game Input failed to disconnect a device. The Device '%s' did not have an associated FGameInputWindowsInputDevice!"), *UE::GameInput::LexToString(Device));
+		}
+	}
+	else
+	{
+		UE_LOG(LogGameInput, Warning, TEXT("Game Input failed to disconnect a device. The Device was null! %s"), *UE::GameInput::LexToString(Device));
+	}
+
+	EnumerateCurrentlyConnectedDeviceTypes();
+}
+
+void FGameInputWindowsInputDevice::HandleDeviceConnected(IGameInputDevice* Device, uint64 Timestamp)
+{
+	// get device information
+	const GameInputDeviceInfo* Info = UE::GameInput::GetDeviceInfo(Device);
+	UE_LOG(LogGameInput, Log, TEXT("Game Input Device Connected: %s  of kind: %s"), *UE::GameInput::LexToString(Device), *UE::GameInput::LexToString(Info->supportedInput));
+	
+	FGameInputDeviceContainer* Data = GetOrCreateDeviceData(Device);
+	check(Data);
+
+	// Map this input device to its user	
+	const FInputDeviceId DeviceId = Data->GetDeviceId();
+	
+	// We only want to map this input device to a user if it has a valid Input Device Id.
+	// Everything has an input device id, so if we got one that is invalid
+	// then that means this device has no processors in it because they were explicitly disabled.
+	if (DeviceId.IsValid())
+	{
+		IPlatformInputDeviceMapper& DeviceMapper = IPlatformInputDeviceMapper::Get();
+		
+		const FPlatformUserId UserToAssign = DeviceMapper.GetPlatformUserForNewlyConnectedDevice();
+	
+		const bool bSuccess = DeviceMapper.Internal_MapInputDeviceToUser(DeviceId, UserToAssign, EInputDeviceConnectionState::Connected);
+		if (ensure(bSuccess))
+		{
+			Data->SetPlatformUserId(UserToAssign);
+		}
+
+		UE_LOG(LogGameInput, Log, TEXT("Using PlatformUserId %d and InputDeviceId %d for device %s"),
+			UserToAssign.GetInternalId(),
+			DeviceId.GetId(),
+			*UE::GameInput::LexToString(Device));
+	}
+	else
+	{
+		UE_LOG(LogGameInput, Log, TEXT("Game Input Device %s had no processors, so it will not be assiged to a platform user. You may need to configure it in the project settings."),
+			*UE::GameInput::LexToString(Device));
+	}
+	
+	EnumerateCurrentlyConnectedDeviceTypes();
+}
+
+FGameInputDeviceContainer* FGameInputWindowsInputDevice::CreateDeviceData(IGameInputDevice* InDevice)
+{
+	TSharedPtr<FGameInputDeviceContainer> Container = MakeShared<FGameInputDeviceContainer>(MessageHandler, InDevice, GetCurrentGameInputKindSupport());
+	Container->InitalizeDeviceProcessors();
+
+	// If this device has any processors assigned to it (meaning that it can send input events) then we want to assign it a new 
+	// input device id. If there are no processors assigned to it, then it can't possibly send input events, so don't bother
+	// mapping it to an input device ID. This can happen if the device is explicitly disallowed for the application 
+	// via the project settings. In this case, we don't really want the rest of the engine to care about it at all
+	const FInputDeviceId AssignedInputDeviceId = 
+		Container->GetNumberOfProcessors() > 0 ?
+		IPlatformInputDeviceMapper::Get().AllocateNewInputDeviceId() :
+		INPUTDEVICEID_NONE;
+	
+	// This is a new device, we need to assign a new input device ID from the platform user
+	Container->SetInputDeviceId(AssignedInputDeviceId);
+
+	DeviceData.Emplace(Container);
+
+	return Container.Get();
+}
+
+#if WITH_EDITOR
+
+void FGameInputWindowsInputDevice::SetupEditorSettingListener()
+{
+	UGameInputDeveloperSettings* Settings = GetMutableDefault<UGameInputDeveloperSettings>();
+	if (!Settings)
+	{
+		return;
+	}
+
+	EditorSettingChangedDelegate = Settings->OnInputSettingChanged.AddRaw(this, &FGameInputWindowsInputDevice::HandleEditorSettingChanged);	
+}
+
+void FGameInputWindowsInputDevice::CleanupEditorSettingListener()
+{
+	if (!EditorSettingChangedDelegate.IsValid())
+	{
+		return;
+	}
+
+	UGameInputDeveloperSettings* Settings = GetMutableDefault<UGameInputDeveloperSettings>();
+	if (!Settings)
+	{
+		return;
+	}
+	
+	Settings->OnInputSettingChanged.Remove(EditorSettingChangedDelegate);
+}
+
+void FGameInputWindowsInputDevice::HandleEditorSettingChanged()
+{
+	for (TSharedPtr<FGameInputDeviceContainer> Device : DeviceData)
+	{
+		// Recreate the device processors for every device
+		Device->RecreateDeviceProcessors(GetCurrentGameInputKindSupport());
+
+		// If this device previously had an invalid DeviceId because it had no processors, but now it does,
+		// then we need to assign a FInputDeviceId because it can now process input
+		// We will never change a device id to be invalid though, because that could cause weirdness where Slate
+		// is expecting input events from a device which no longer exists.
+		if (!Device->GetDeviceId().IsValid() && Device->GetNumberOfProcessors() > 0)
+		{
+			Device->SetInputDeviceId(IPlatformInputDeviceMapper::Get().AllocateNewInputDeviceId());
+		}
+	}
+}
+#endif	// WITH_EDITOR
+
+#endif	//#if GAME_INPUT_SUPPORT
\ No newline at end of file
diff --git a/Engine/Plugins/GameInputWindows/Source/GameInputWindows/Private/GameInputWindowsDevice.h b/Engine/Plugins/GameInputWindows/Source/GameInputWindows/Private/GameInputWindowsDevice.h
new file mode 100644
--- /dev/null
+++ b/Engine/Plugins/GameInputWindows/Source/GameInputWindows/Private/GameInputWindowsDevice.h
@@ -0,0 +1,36 @@
+// Copyright Epic Games, Inc. All Rights Reserved.
+
+#pragma once
+
+#if GAME_INPUT_SUPPORT
+
+#include "IGameInputDeviceInterface.h"
+
+class FGameInputWindowsInputDevice : public IGameInputDeviceInterface
+{
+public:
+	explicit FGameInputWindowsInputDevice(const TSharedRef<FGenericApplicationMessageHandler>& InMessageHandler, struct IGameInput* InGameInput);
+	virtual ~FGameInputWindowsInputDevice() override;
+
+	void SetGameInputAndReinitialize(IGameInput* InGameInput);
+
+protected:
+
+	//~Begin IGameInputDeviceInterface interface	
+	virtual GameInputKind GetCurrentGameInputKindSupport() const override;
+	virtual void HandleDeviceDisconnected(IGameInputDevice* Device, uint64 Timestamp) override;
+	virtual void HandleDeviceConnected(IGameInputDevice* Device, uint64 Timestamp) override;
+	virtual FGameInputDeviceContainer* CreateDeviceData(IGameInputDevice* InDevice) override;
+	//~End IGameInputDeviceInterface interface
+	
+private:
+#if WITH_EDITOR
+	void SetupEditorSettingListener();
+	void CleanupEditorSettingListener();
+	void HandleEditorSettingChanged();
+
+	FDelegateHandle EditorSettingChangedDelegate;
+#endif
+};
+
+#endif	// #if GAME_INPUT_SUPPORT
\ No newline at end of file
diff --git a/Engine/Plugins/GameInputWindows/Source/GameInputWindows/Private/GameInputWindowsModule.cpp b/Engine/Plugins/GameInputWindows/Source/GameInputWindows/Private/GameInputWindowsModule.cpp
new file mode 100644
--- /dev/null
+++ b/Engine/Plugins/GameInputWindows/Source/GameInputWindows/Private/GameInputWindowsModule.cpp
@@ -0,0 +1,59 @@
+// Copyright Epic Games, Inc. All Rights Reserved.
+
+#include "GameInputWindowsModule.h"
+
+#include "CoreGlobals.h"
+#include "GameInputWindowsDevice.h"
+#include "GameInputBaseModule.h"
+#include "GameInputLogging.h"
+#include "Misc/App.h"
+#include "Modules/ModuleManager.h"
+#include "GameInputBaseIncludes.h"
+
+namespace UE::GameInput
+{
+	/** The name of this modular feature plugin. */
+	static const FName GameInputWindowsFeatureName = TEXT("GameInputWindows");
+}
+
+FGameInputWindowsModule& FGameInputWindowsModule::Get()
+{
+	return FModuleManager::LoadModuleChecked<FGameInputWindowsModule>(UE::GameInput::GameInputWindowsFeatureName);
+}
+
+bool FGameInputWindowsModule::IsAvailable()
+{
+	return FModuleManager::Get().IsModuleLoaded(UE::GameInput::GameInputWindowsFeatureName);
+}
+
+TSharedPtr<IInputDevice> FGameInputWindowsModule::CreateInputDevice(const TSharedRef<FGenericApplicationMessageHandler>& InMessageHandler)
+{
+#if GAME_INPUT_SUPPORT
+	//Get FGameInputBaseModule before anything else to make sure that it is loaded
+	FGameInputBaseModule& BaseModule = FGameInputBaseModule::Get();
+	IGameInput* BaseGameInputPtr = FGameInputBaseModule::GetGameInput();
+	TSharedPtr<FGameInputWindowsInputDevice> DevicePtr = MakeShared<FGameInputWindowsInputDevice>(InMessageHandler, BaseGameInputPtr);
+	
+	if (BaseGameInputPtr)
+	{
+		DevicePtr->Initialize();
+	}
+	else
+	{	
+		BaseModule.OnGameInputCreation.AddLambda([&DevicePtr](IGameInput* GameInputPtr)
+		{
+			if (ensure(DevicePtr.IsValid()))
+			{
+				DevicePtr->SetGameInputAndReinitialize(GameInputPtr);
+			}
+		});	
+	}
+	return DevicePtr;
+	
+#else
+	UE_LOG(LogGameInput, Error, TEXT("[FGameInputWindowsModule] Failed to create a GameInput device! GAME_INPUT_SUPPORT is false! See GameInputBase.build.cs and the build log for more info."));
+	return nullptr;
+#endif	// GAME_INPUT_SUPPORT
+}
+
+IMPLEMENT_MODULE(FGameInputWindowsModule, GameInputWindows)
\ No newline at end of file
diff --git a/Engine/Plugins/GameInputWindows/Source/GameInputWindows/Private/GameInputWindowsModule.h b/Engine/Plugins/GameInputWindows/Source/GameInputWindows/Private/GameInputWindowsModule.h
new file mode 100644
--- /dev/null
+++ b/Engine/Plugins/GameInputWindows/Source/GameInputWindows/Private/GameInputWindowsModule.h
@@ -0,0 +1,21 @@
+// Copyright Epic Games, Inc. All Rights Reserved.
+
+#pragma once
+
+#include "IInputDeviceModule.h"
+
+/**
+ * Input Device module that will create the Game Input Device to handle input on Windows platforms.
+ */
+class FGameInputWindowsModule : public IInputDeviceModule
+{
+public:
+	static FGameInputWindowsModule& Get();
+
+	/** Returns true if this module is loaded (aka available) by the FModuleManager */
+	static bool IsAvailable();
+
+	//~ Begin IInputDeviceModule interface
+	virtual TSharedPtr<IInputDevice> CreateInputDevice(const TSharedRef<FGenericApplicationMessageHandler>& InMessageHandler) override;		
+	//~ End IInputDeviceModule interface
+};
\ No newline at end of file
