[Game Input] Fix a dangling pointer in lambda capture

#rnx
#Synced-CL 51759702

diff --git a/Engine/Plugins/GameInputWindows/Source/GameInputWindows/Private/GameInputWindowsModule.cpp b/Engine/Plugins/GameInputWindows/Source/GameInputWindows/Private/GameInputWindowsModule.cpp
--- a/Engine/Plugins/GameInputWindows/Source/GameInputWindows/Private/GameInputWindowsModule.cpp
+++ b/Engine/Plugins/GameInputWindows/Source/GameInputWindows/Private/GameInputWindowsModule.cpp
@@ -39,14 +39,15 @@ TSharedPtr<IInputDevice> FGameInputWindowsModule::CreateInputDevice(const TShare
 		DevicePtr->Initialize();
 	}
 	else
-	{	
-		BaseModule.OnGameInputCreation.AddLambda([&DevicePtr](IGameInput* GameInputPtr)
+	{
+		TWeakPtr<FGameInputWindowsInputDevice> WeakDevicePtr = DevicePtr;
+		BaseModule.OnGameInputCreation.AddLambda([WeakDevicePtr](IGameInput* GameInputPtr)
 		{
-			if (ensure(DevicePtr.IsValid()))
+			if (TSharedPtr<FGameInputWindowsInputDevice> PinnedDevice = WeakDevicePtr.Pin())
 			{
-				DevicePtr->SetGameInputAndReinitialize(GameInputPtr);
+				PinnedDevice->SetGameInputAndReinitialize(GameInputPtr);
 			}
-		});	
+		});
 	}
 	return DevicePtr;
 	
