[Game Input] Add default hardware device Id data for Windows targets

This will let you query the device type from the UInputDeviceSubsystem for supported devices.

#rnx
#Synced-CL 51946316

diff --git a/Engine/Plugins/GameInputWindows/Config/Input.ini b/Engine/Plugins/GameInputWindows/Config/Input.ini
new file mode 100644
--- /dev/null
+++ b/Engine/Plugins/GameInputWindows/Config/Input.ini
@@ -0,0 +1,7 @@
+[InputPlatformSettings_Windows InputPlatformSettings]
++HardwareDevices=(InputClassName="GameInput",HardwareDeviceIdentifier="Virtual",PrimaryDeviceType=Gamepad,SupportedFeaturesMask=2055)
++HardwareDevices=(InputClassName="GameInput",HardwareDeviceIdentifier="Aggregate",PrimaryDeviceType=Gamepad,SupportedFeaturesMask=4)
++HardwareDevices=(InputClassName="GameInput",HardwareDeviceIdentifier="XboxOne",PrimaryDeviceType=Gamepad,SupportedFeaturesMask=260)
++HardwareDevices=(InputClassName="GameInput",HardwareDeviceIdentifier="Xbox360",PrimaryDeviceType=Gamepad,SupportedFeaturesMask=260)
++HardwareDevices=(InputClassName="GameInput",HardwareDeviceIdentifier="Hid",PrimaryDeviceType=Unspecified,SupportedFeaturesMask=0)
++HardwareDevices=(InputClassName="GameInput",HardwareDeviceIdentifier="I8042",PrimaryDeviceType=Unspecified,SupportedFeaturesMask=0)
\ No newline at end of file
