[Game Input] Make the GameInput::HID device type a gamepad input device type by default.

#Synced-CL 52059832

#Synced-CL 52059832

diff --git a/Engine/Plugins/GameInputWindows/Config/Input.ini b/Engine/Plugins/GameInputWindows/Config/Input.ini
--- a/Engine/Plugins/GameInputWindows/Config/Input.ini
+++ b/Engine/Plugins/GameInputWindows/Config/Input.ini
@@ -3,5 +3,5 @@
 +HardwareDevices=(InputClassName="GameInput",HardwareDeviceIdentifier="Aggregate",PrimaryDeviceType=Gamepad,SupportedFeaturesMask=4)
 +HardwareDevices=(InputClassName="GameInput",HardwareDeviceIdentifier="XboxOne",PrimaryDeviceType=Gamepad,SupportedFeaturesMask=260)
 +HardwareDevices=(InputClassName="GameInput",HardwareDeviceIdentifier="Xbox360",PrimaryDeviceType=Gamepad,SupportedFeaturesMask=260)
-+HardwareDevices=(InputClassName="GameInput",HardwareDeviceIdentifier="Hid",PrimaryDeviceType=Unspecified,SupportedFeaturesMask=0)
++HardwareDevices=(InputClassName="GameInput",HardwareDeviceIdentifier="Hid",PrimaryDeviceType=Gamepad,SupportedFeaturesMask=0)
 +HardwareDevices=(InputClassName="GameInput",HardwareDeviceIdentifier="I8042",PrimaryDeviceType=Unspecified,SupportedFeaturesMask=0)
\ No newline at end of file
